Make BufferVisualization can config with MaterialInstances.

diff --git a/Engine/Plugins/Compositing/Composure/Source/Composure/Private/MovieScene/MovieSceneComposureExportSectionTemplate.cpp b/Engine/Plugins/Compositing/Composure/Source/Composure/Private/MovieScene/MovieSceneComposureExportSectionTemplate.cpp
--- a/Engine/Plugins/Compositing/Composure/Source/Composure/Private/MovieScene/MovieSceneComposureExportSectionTemplate.cpp
+++ b/Engine/Plugins/Compositing/Composure/Source/Composure/Private/MovieScene/MovieSceneComposureExportSectionTemplate.cpp
@@ -31,7 +31,7 @@ struct FBufferVisualizationIterator
 	/** Reference to the post processing settings to add vialization materials to */
 	FFinalPostProcessSettings& FinalPostProcessSettings;
 
-	void ProcessValue(const FString& InName, UMaterial* Material, const FText& InText)
+	void ProcessValue(const FString& InName, UMaterialInterface* Material, const FText& InText)
 	{
 		if (BuffersToCapture.Contains(InName) || BuffersToCapture.Contains(InText.ToString()))
 		{
diff --git a/Engine/Source/Editor/CommonMenuExtensions/Private/BufferVisualizationMenuCommands.cpp b/Engine/Source/Editor/CommonMenuExtensions/Private/BufferVisualizationMenuCommands.cpp
--- a/Engine/Source/Editor/CommonMenuExtensions/Private/BufferVisualizationMenuCommands.cpp
+++ b/Engine/Source/Editor/CommonMenuExtensions/Private/BufferVisualizationMenuCommands.cpp
@@ -16,7 +16,7 @@ namespace
 {
 	struct FMaterialIteratorWithLambda
 	{
-		typedef TFunction<void(const FString&, const UMaterial*, const FText&)> LambdaFunctionType;
+		typedef TFunction<void(const FString&, const UMaterialInterface*, const FText&)> LambdaFunctionType;
 		LambdaFunctionType Lambda;
 
 		inline FMaterialIteratorWithLambda(const LambdaFunctionType& InLambda)
@@ -24,7 +24,7 @@ namespace
 		{
 		}
 
-		inline void ProcessValue(const FString& InMaterialName, const UMaterial* InMaterial, const FText& InDisplayName)
+		inline void ProcessValue(const FString& InMaterialName, const UMaterialInterface* InMaterial, const FText& InDisplayName)
 		{
 			Lambda(InMaterialName, InMaterial, InDisplayName);
 		}
@@ -74,7 +74,7 @@ void FBufferVisualizationMenuCommands::CreateOverviewCommand()
 
 void FBufferVisualizationMenuCommands::CreateVisualizationCommands()
 {
-	FMaterialIteratorWithLambda Iterator([this](const FString& InMaterialName, const UMaterial* InMaterial, const FText& InDisplayName)
+	FMaterialIteratorWithLambda Iterator([this](const FString& InMaterialName, const UMaterialInterface* InMaterial, const FText& InDisplayName)
 	{
 		const FName CommandName = NameFromMaterial(InMaterialName);
 
@@ -115,7 +115,7 @@ void FBufferVisualizationMenuCommands::AddVisualizationCommandsToMenu(FMenuBuild
 
 	const TBufferVisualizationModeCommandMap& Commands = CommandMap;
 
-	FMaterialIteratorWithLambda Iterator([&Menu, &Commands](const FString& InMaterialName, const UMaterial* InMaterial, const FText& InDisplayName)
+	FMaterialIteratorWithLambda Iterator([&Menu, &Commands](const FString& InMaterialName, const UMaterialInterface* InMaterial, const FText& InDisplayName)
 	{
 		const FName CommandName = NameFromMaterial(InMaterialName);
 
diff --git a/Engine/Source/Editor/DetailCustomizations/Private/RenderPassesCustomization.cpp b/Engine/Source/Editor/DetailCustomizations/Private/RenderPassesCustomization.cpp
--- a/Engine/Source/Editor/DetailCustomizations/Private/RenderPassesCustomization.cpp
+++ b/Engine/Source/Editor/DetailCustomizations/Private/RenderPassesCustomization.cpp
@@ -83,7 +83,7 @@ public:
 
 			FIterator(TMap<FString, FText>& InRenderPasses) : RenderPasses(InRenderPasses) {}
 
-			void ProcessValue(const FString& InName, UMaterial* Material, const FText& InText)
+			void ProcessValue(const FString& InName, UMaterialInterface* Material, const FText& InText)
 			{
 				RenderPasses.Add(InName, InText);
 			}
diff --git a/Engine/Source/Runtime/Engine/Private/BufferVisualizationData.cpp b/Engine/Source/Runtime/Engine/Private/BufferVisualizationData.cpp
--- a/Engine/Source/Runtime/Engine/Private/BufferVisualizationData.cpp
+++ b/Engine/Source/Runtime/Engine/Private/BufferVisualizationData.cpp
@@ -39,7 +39,7 @@ void FBufferVisualizationData::Initialize()
 					FString MaterialName;
 					if( FParse::Value( *It.Value().GetValue(), TEXT("Material="), MaterialName, true ) )
 					{
-						UMaterial* Material = LoadObject<UMaterial>(NULL, *MaterialName);
+						UMaterialInterface* Material = LoadObject<UMaterialInterface>(NULL, *MaterialName);
 			
 						if (Material)
 						{
@@ -73,7 +73,7 @@ void FBufferVisualizationData::ConfigureConsoleCommand()
 		{
 		}
 
-		void ProcessValue(const FString& InMaterialName, const UMaterial* InMaterial, const FText& InDisplayName)
+		void ProcessValue(const FString& InMaterialName, const UMaterialInterface* InMaterial, const FText& InDisplayName)
 		{
 			Message += FString(TEXT("\n  "));
 			Message += InMaterialName;
@@ -102,7 +102,7 @@ void FBufferVisualizationData::ConfigureConsoleCommand()
 		ECVF_Default);
 }
 
-UMaterial* FBufferVisualizationData::GetMaterial(FName InMaterialName)
+UMaterialInterface* FBufferVisualizationData::GetMaterial(FName InMaterialName)
 {
 	Record* Result = MaterialMap.Find(InMaterialName);
 	if (Result)
@@ -125,7 +125,7 @@ bool FBufferVisualizationData::IsDifferentToCurrentOverviewMaterialNames(const F
 	return InNameList != CurrentOverviewMaterialNames;
 }
 
-TArray<UMaterial*>& FBufferVisualizationData::GetOverviewMaterials()
+TArray<UMaterialInterface*>& FBufferVisualizationData::GetOverviewMaterials()
 {
 	return OverviewMaterials;
 }
diff --git a/Engine/Source/Runtime/Engine/Private/DebugCameraController.cpp b/Engine/Source/Runtime/Engine/Private/DebugCameraController.cpp
--- a/Engine/Source/Runtime/Engine/Private/DebugCameraController.cpp
+++ b/Engine/Source/Runtime/Engine/Private/DebugCameraController.cpp
@@ -960,7 +960,7 @@ FString ADebugCameraController::GetBufferMaterialName(const FString& InBufferNam
 {
 	if (!InBufferName.IsEmpty())
 	{
-		if (UMaterial* Material = GetBufferVisualizationData().GetMaterial(*InBufferName))
+		if (UMaterialInterface* Material = GetBufferVisualizationData().GetMaterial(*InBufferName))
 		{
 			return Material->GetName();
 		}
diff --git a/Engine/Source/Runtime/Engine/Private/SceneView.cpp b/Engine/Source/Runtime/Engine/Private/SceneView.cpp
--- a/Engine/Source/Runtime/Engine/Private/SceneView.cpp
+++ b/Engine/Source/Runtime/Engine/Private/SceneView.cpp
@@ -2068,7 +2068,7 @@ void FSceneView::ConfigureBufferVisualizationSettings()
 
 				// Lookup this material from the list that was parsed out of the global ini file
 				Left.TrimStartInline();
-				UMaterial* Material = BufferVisualizationData.GetMaterial(*Left);
+				UMaterialInterface* Material = BufferVisualizationData.GetMaterial(*Left);
 
 				if (Material == NULL && Left.Len() > 0)
 				{
@@ -2084,7 +2084,7 @@ void FSceneView::ConfigureBufferVisualizationSettings()
 		}
 
 		// Copy current material list into settings material list
-		for (TArray<UMaterial*>::TConstIterator It = BufferVisualizationData.GetOverviewMaterials().CreateConstIterator(); It; ++It)
+		for (TArray<UMaterialInterface*>::TConstIterator It = BufferVisualizationData.GetOverviewMaterials().CreateConstIterator(); It; ++It)
 		{
 			FinalPostProcessSettings.BufferVisualizationOverviewMaterials.Add(*It);
 		}
diff --git a/Engine/Source/Runtime/Engine/Public/BufferVisualizationData.h b/Engine/Source/Runtime/Engine/Public/BufferVisualizationData.h
--- a/Engine/Source/Runtime/Engine/Public/BufferVisualizationData.h
+++ b/Engine/Source/Runtime/Engine/Public/BufferVisualizationData.h
@@ -3,7 +3,7 @@
 
 #include "CoreMinimal.h"
 
-class UMaterial;
+class UMaterialInterface;
 
 class FBufferVisualizationData
 {
@@ -22,14 +22,14 @@ public:
 	bool IsInitialized() const { return bIsInitialized; }
 
 	/** Get a named material from the available material map **/
-	ENGINE_API UMaterial* GetMaterial(FName InMaterialName);
+	ENGINE_API UMaterialInterface* GetMaterial(FName InMaterialName);
 
 	/** We cache the overview material name list from the console command here, so all dynamically created views can re-use the existing cached list of materials */
 	void SetCurrentOverviewMaterialNames(const FString& InNameList);
 	bool IsDifferentToCurrentOverviewMaterialNames(const FString& InNameList);
 
 	/** Access the list of materials currently in use by the buffer visualization overview */
-	TArray<UMaterial*>& GetOverviewMaterials();
+	TArray<UMaterialInterface*>& GetOverviewMaterials();
 
 	/** Iterator function for iterating over available materials */
 	template<class T> void IterateOverAvailableMaterials(T& Iterator) const
@@ -54,7 +54,7 @@ private:
 	{
 		FString Name;
 		FText DisplayName;
-		UMaterial* Material;
+		UMaterialInterface* Material;
 	};
 
 	/** Mapping of FName (first parameter in ini file material list) to a material record */
@@ -67,7 +67,7 @@ private:
 	FString CurrentOverviewMaterialNames;
 
 	/** List of material currently in use by the buffer visualization overview */
-	TArray<UMaterial*> OverviewMaterials;
+	TArray<UMaterialInterface*> OverviewMaterials;
 	
 	/** Storage for console variable documentation strings **/
 	FString ConsoleDocumentationVisualizationMode;
diff --git a/Engine/Source/Runtime/MovieSceneCapture/Private/CompositionGraphCaptureProtocol.cpp b/Engine/Source/Runtime/MovieSceneCapture/Private/CompositionGraphCaptureProtocol.cpp
--- a/Engine/Source/Runtime/MovieSceneCapture/Private/CompositionGraphCaptureProtocol.cpp
+++ b/Engine/Source/Runtime/MovieSceneCapture/Private/CompositionGraphCaptureProtocol.cpp
@@ -110,7 +110,7 @@ struct FFrameCaptureViewExtension : public FSceneViewExtensionBase
 				: FinalPostProcessSettings(InFinalPostProcessSettings), RenderPasses(InRenderPasses)
 			{}
 
-			void ProcessValue(const FString& InName, UMaterial* Material, const FText& InText)
+			void ProcessValue(const FString& InName, UMaterialInterface* Material, const FText& InText)
 			{
 				if (!RenderPasses.Num() || RenderPasses.Contains(InName) || RenderPasses.Contains(InText.ToString()))
 				{
diff --git a/Engine/Source/Runtime/Renderer/Private/PostProcess/PostProcessMaterial.cpp b/Engine/Source/Runtime/Renderer/Private/PostProcess/PostProcessMaterial.cpp
--- a/Engine/Source/Runtime/Renderer/Private/PostProcess/PostProcessMaterial.cpp
+++ b/Engine/Source/Runtime/Renderer/Private/PostProcess/PostProcessMaterial.cpp
@@ -501,7 +501,8 @@ FPostProcessMaterialChain GetPostProcessMaterialChain(const FViewInfo& View, EBl
 
 	if (ViewFamily.EngineShowFlags.VisualizeBuffer)
 	{
-		UMaterial* Material = GetBufferVisualizationData().GetMaterial(View.CurrentBufferVisualizationMode);
+		UMaterialInterface* VisMaterial = GetBufferVisualizationData().GetMaterial(View.CurrentBufferVisualizationMode);
+		UMaterial* Material = VisMaterial ? VisMaterial->GetMaterial() : nullptr;
 
 		if (Material && Material->BlendableLocation == Location)
 		{
