Viewport show flags menu converted
#rnx

diff --git a/Engine/Source/Editor/CommonMenuExtensions/CommonMenuExtensions.Build.cs b/Engine/Source/Editor/CommonMenuExtensions/CommonMenuExtensions.Build.cs
--- a/Engine/Source/Editor/CommonMenuExtensions/CommonMenuExtensions.Build.cs
+++ b/Engine/Source/Editor/CommonMenuExtensions/CommonMenuExtensions.Build.cs
@@ -33,7 +33,8 @@ public class CommonMenuExtensions : ModuleRules
 				"EngineSettings",
 				"HierarchicalLODOutliner",
 				"HierarchicalLODUtilities",
-				"MaterialShaderQualitySettings"
+				"MaterialShaderQualitySettings",
+				"ToolMenus",
 			}
 		);
 
diff --git a/Engine/Source/Editor/CommonMenuExtensions/Private/ShowFlagMenuCommands.cpp b/Engine/Source/Editor/CommonMenuExtensions/Private/ShowFlagMenuCommands.cpp
--- a/Engine/Source/Editor/CommonMenuExtensions/Private/ShowFlagMenuCommands.cpp
+++ b/Engine/Source/Editor/CommonMenuExtensions/Private/ShowFlagMenuCommands.cpp
@@ -3,6 +3,8 @@
 #include "ShowFlagMenuCommands.h"
 #include "EditorShowFlags.h"
 #include "EditorStyleSet.h"
+#include "Framework/MultiBox/MultiBoxBuilder.h"
+#include "ToolMenus.h"
 
 #define LOCTEXT_NAMESPACE "ShowFlagMenuCommands"
 
@@ -41,7 +43,7 @@ FShowFlagMenuCommands::FShowFlagMenuCommands()
 {
 }
 
-void FShowFlagMenuCommands::BuildShowFlagsMenu(FMenuBuilder& Menu, const FShowFlagFilter& Filter) const
+void FShowFlagMenuCommands::BuildShowFlagsMenu(UToolMenu* Menu, const FShowFlagFilter& Filter) const
 {
 	check(bCommandsInitialised);
 
@@ -53,20 +55,19 @@ void FShowFlagMenuCommands::BuildShowFlagsMenu(FMenuBuilder& Menu, const FShowFl
 
 	CreateCommonShowFlagMenuItems(Menu, Filter);
 
-	Menu.BeginSection("LevelViewportShowFlags", LOCTEXT("AllShowFlagHeader", "All Show Flags"));
 	{
-		CreateSubMenuIfRequired(Menu, Filter, SFG_PostProcess, LOCTEXT("PostProcessShowFlagsMenu", "Post Processing"), LOCTEXT("PostProcessShowFlagsMenu_ToolTip", "Post process show flags"));
-		CreateSubMenuIfRequired(Menu, Filter, SFG_LightTypes, LOCTEXT("LightTypesShowFlagsMenu", "Light Types"), LOCTEXT("LightTypesShowFlagsMenu_ToolTip", "Light Types show flags"));
-		CreateSubMenuIfRequired(Menu, Filter, SFG_LightingComponents, LOCTEXT("LightingComponentsShowFlagsMenu", "Lighting Components"), LOCTEXT("LightingComponentsShowFlagsMenu_ToolTip", "Lighting Components show flags"));
-		CreateSubMenuIfRequired(Menu, Filter, SFG_LightingFeatures, LOCTEXT("LightingFeaturesShowFlagsMenu", "Lighting Features"), LOCTEXT("LightingFeaturesShowFlagsMenu_ToolTip", "Lighting Features show flags"));
-		CreateSubMenuIfRequired(Menu, Filter, SFG_Developer, LOCTEXT("DeveloperShowFlagsMenu", "Developer"), LOCTEXT("DeveloperShowFlagsMenu_ToolTip", "Developer show flags"));
-		CreateSubMenuIfRequired(Menu, Filter, SFG_Visualize, LOCTEXT("VisualizeShowFlagsMenu", "Visualize"), LOCTEXT("VisualizeShowFlagsMenu_ToolTip", "Visualize show flags"));
-		CreateSubMenuIfRequired(Menu, Filter, SFG_Advanced, LOCTEXT("AdvancedShowFlagsMenu", "Advanced"), LOCTEXT("AdvancedShowFlagsMenu_ToolTip", "Advanced show flags"));
+		FToolMenuSection& Section = Menu->AddSection("LevelViewportShowFlags", LOCTEXT("AllShowFlagHeader", "All Show Flags"));;
+		CreateSubMenuIfRequired(Section, Filter, SFG_PostProcess, "SFG_PostProcess", LOCTEXT("PostProcessShowFlagsMenu", "Post Processing"), LOCTEXT("PostProcessShowFlagsMenu_ToolTip", "Post process show flags"));
+		CreateSubMenuIfRequired(Section, Filter, SFG_LightTypes, "SFG_LightTypes", LOCTEXT("LightTypesShowFlagsMenu", "Light Types"), LOCTEXT("LightTypesShowFlagsMenu_ToolTip", "Light Types show flags"));
+		CreateSubMenuIfRequired(Section, Filter, SFG_LightingComponents, "SFG_LightingComponents", LOCTEXT("LightingComponentsShowFlagsMenu", "Lighting Components"), LOCTEXT("LightingComponentsShowFlagsMenu_ToolTip", "Lighting Components show flags"));
+		CreateSubMenuIfRequired(Section, Filter, SFG_LightingFeatures, "SFG_LightingFeatures", LOCTEXT("LightingFeaturesShowFlagsMenu", "Lighting Features"), LOCTEXT("LightingFeaturesShowFlagsMenu_ToolTip", "Lighting Features show flags"));
+		CreateSubMenuIfRequired(Section, Filter, SFG_Developer, "SFG_Developer", LOCTEXT("DeveloperShowFlagsMenu", "Developer"), LOCTEXT("DeveloperShowFlagsMenu_ToolTip", "Developer show flags"));
+		CreateSubMenuIfRequired(Section, Filter, SFG_Visualize, "SFG_Visualize", LOCTEXT("VisualizeShowFlagsMenu", "Visualize"), LOCTEXT("VisualizeShowFlagsMenu_ToolTip", "Visualize show flags"));
+		CreateSubMenuIfRequired(Section, Filter, SFG_Advanced, "SFG_Advanced", LOCTEXT("AdvancedShowFlagsMenu", "Advanced"), LOCTEXT("AdvancedShowFlagsMenu_ToolTip", "Advanced show flags"));
 	}
-	Menu.EndSection();
 }
 
-void FShowFlagMenuCommands::CreateCommonShowFlagMenuItems(FMenuBuilder& Menu, const FShowFlagFilter& Filter) const
+void FShowFlagMenuCommands::CreateCommonShowFlagMenuItems(UToolMenu* Menu, const FShowFlagFilter& Filter) const
 {
 	const FShowFlagFilter::FGroupedShowFlagIndices& GroupedFlagIndices = Filter.GetFilteredIndices();
 	const TArray<uint32>& FlagIndices = GroupedFlagIndices[SFG_Normal];
@@ -76,20 +77,20 @@ void FShowFlagMenuCommands::CreateCommonShowFlagMenuItems(FMenuBuilder& Menu, co
 		return;
 	}
 
-	Menu.BeginSection("ShowFlagsMenuSectionCommon", LOCTEXT("CommonShowFlagHeader", "Common Show Flags"));
 	{
+		FToolMenuSection& Section = Menu->AddSection("ShowFlagsMenuSectionCommon", LOCTEXT("CommonShowFlagHeader", "Common Show Flags"));
 		for (int32 ArrayIndex = 0; ArrayIndex < FlagIndices.Num(); ++ArrayIndex)
 		{
 			const uint32 FlagIndex = FlagIndices[ArrayIndex];
 			const FShowFlagCommand& ShowFlagCommand = ShowFlagCommands[FlagIndex];
 
-			Menu.AddMenuEntry(ShowFlagCommand.ShowMenuItem, NAME_None, ShowFlagCommand.LabelOverride);
+			ensure(Section.FindEntry(ShowFlagCommand.ShowMenuItem->GetCommandName()) == nullptr);
+			Section.AddMenuEntry(*FString::Printf(TEXT("Common_%s"), *ShowFlagCommand.ShowMenuItem->GetCommandName().ToString()), ShowFlagCommand.ShowMenuItem, ShowFlagCommand.LabelOverride);
 		}
 	}
-	Menu.EndSection();
 }
 
-void FShowFlagMenuCommands::CreateSubMenuIfRequired(FMenuBuilder& Menu, const FShowFlagFilter& Filter, EShowFlagGroup Group, const FText& MenuLabel, const FText& ToolTip) const
+void FShowFlagMenuCommands::CreateSubMenuIfRequired(FToolMenuSection& Section, const FShowFlagFilter& Filter, EShowFlagGroup Group, const FName SubMenuName, const FText& MenuLabel, const FText& ToolTip) const
 {
 	const FShowFlagFilter::FGroupedShowFlagIndices& GroupedFlagIndices = Filter.GetFilteredIndices();
 	const TArray<uint32>& FlagIndices = GroupedFlagIndices[Group];
@@ -99,24 +100,26 @@ void FShowFlagMenuCommands::CreateSubMenuIfRequired(FMenuBuilder& Menu, const FS
 		return;
 	}
 
-	Menu.AddSubMenu(MenuLabel, ToolTip, FNewMenuDelegate::CreateStatic(&FShowFlagMenuCommands::StaticCreateShowFlagsSubMenu, FlagIndices, 0));
+	Section.AddSubMenu(SubMenuName, MenuLabel, ToolTip, FNewToolMenuDelegate::CreateStatic(&FShowFlagMenuCommands::StaticCreateShowFlagsSubMenu, FlagIndices, 0));
 }
 
-void FShowFlagMenuCommands::CreateShowFlagsSubMenu(FMenuBuilder& Menu, TArray<uint32> FlagIndices, int32 EntryOffset) const
+void FShowFlagMenuCommands::CreateShowFlagsSubMenu(UToolMenu* Menu, TArray<uint32> FlagIndices, int32 EntryOffset) const
 {
 	// Generate entries for the standard show flags.
 	// Assumption: the first 'n' entries types like 'Show All' and 'Hide All' buttons, so insert a separator after them.
 
+	FToolMenuSection& Section = Menu->AddSection("Section");
 	for (int32 ArrayIndex = 0; ArrayIndex < FlagIndices.Num(); ++ArrayIndex)
 	{
 		const uint32 FlagIndex = FlagIndices[ArrayIndex];
 		const FShowFlagCommand& ShowFlagCommand = ShowFlagCommands[FlagIndex];
 
-		Menu.AddMenuEntry(ShowFlagCommand.ShowMenuItem, NAME_None, ShowFlagCommand.LabelOverride);
+		ensure(Section.FindEntry(ShowFlagCommand.ShowMenuItem->GetCommandName()) == nullptr);
+		Section.AddMenuEntry(ShowFlagCommand.ShowMenuItem->GetCommandName(), ShowFlagCommand.ShowMenuItem, ShowFlagCommand.LabelOverride);
 
 		if (ArrayIndex == EntryOffset - 1)
 		{
-			Menu.AddMenuSeparator();
+			Section.AddMenuSeparator(NAME_None);
 		}
 	}
 }
@@ -148,9 +151,9 @@ void FShowFlagMenuCommands::CreateShowFlagCommands()
 	}
 }
 
-void FShowFlagMenuCommands::StaticCreateShowFlagsSubMenu(FMenuBuilder& MenuBuilder, TArray<uint32> FlagIndices, int32 EntryOffset)
+void FShowFlagMenuCommands::StaticCreateShowFlagsSubMenu(UToolMenu* Menu, TArray<uint32> FlagIndices, int32 EntryOffset)
 {
-	FShowFlagMenuCommands::Get().CreateShowFlagsSubMenu(MenuBuilder, FlagIndices, EntryOffset);
+	FShowFlagMenuCommands::Get().CreateShowFlagsSubMenu(Menu, FlagIndices, EntryOffset);
 }
 
 FSlateIcon FShowFlagMenuCommands::GetShowFlagIcon(const FShowFlagData& Flag) const
diff --git a/Engine/Source/Editor/CommonMenuExtensions/Public/ShowFlagMenuCommands.h b/Engine/Source/Editor/CommonMenuExtensions/Public/ShowFlagMenuCommands.h
--- a/Engine/Source/Editor/CommonMenuExtensions/Public/ShowFlagMenuCommands.h
+++ b/Engine/Source/Editor/CommonMenuExtensions/Public/ShowFlagMenuCommands.h
@@ -7,10 +7,12 @@
 #include "Templates/SharedPointer.h"
 #include "Framework/Commands/UICommandInfo.h"
 #include "ShowFlagFilter.h"
-#include "Framework/MultiBox/MultiBoxBuilder.h"
 #include "Templates/Function.h"
 #include "EditorViewportClient.h"
 
+class UToolMenu;
+struct FToolMenuSection;
+
 class COMMONMENUEXTENSIONS_API FShowFlagMenuCommands : public TCommands<FShowFlagMenuCommands>
 {
 public:
@@ -40,19 +42,19 @@ public:
 	virtual void RegisterCommands() override;
 
 	void BindCommands(FUICommandList& CommandList, const TSharedPtr<FEditorViewportClient>& Client) const;
-	void BuildShowFlagsMenu(FMenuBuilder& Menu, const FShowFlagFilter& Filter = FShowFlagFilter(FShowFlagFilter::IncludeAllFlagsByDefault)) const;
+	void BuildShowFlagsMenu(UToolMenu* Menu, const FShowFlagFilter& Filter = FShowFlagFilter(FShowFlagFilter::IncludeAllFlagsByDefault)) const;
 
 private:
-	static void StaticCreateShowFlagsSubMenu(FMenuBuilder& MenuBuilder, TArray<uint32> FlagIndices, int32 EntryOffset);
+	static void StaticCreateShowFlagsSubMenu(UToolMenu* Menu, TArray<uint32> FlagIndices, int32 EntryOffset);
 	static void ToggleShowFlag(const TSharedPtr<FEditorViewportClient>& Client, FEngineShowFlags::EShowFlag EngineShowFlagIndex);
 	static bool IsShowFlagEnabled(const TSharedPtr<FEditorViewportClient>& Client, FEngineShowFlags::EShowFlag EngineShowFlagIndex);
 
 	FSlateIcon GetShowFlagIcon(const FShowFlagData& Flag) const;
 
 	void CreateShowFlagCommands();
-	void CreateCommonShowFlagMenuItems(FMenuBuilder& Menu, const FShowFlagFilter& Filter) const;
-	void CreateSubMenuIfRequired(FMenuBuilder& Menu, const FShowFlagFilter& Filter, EShowFlagGroup Group, const FText& MenuLabel, const FText& ToolTip) const;
-	void CreateShowFlagsSubMenu(FMenuBuilder& MenuBuilder, TArray<uint32> MenuCommands, int32 EntryOffset) const;
+	void CreateCommonShowFlagMenuItems(UToolMenu* Menu, const FShowFlagFilter& Filter) const;
+	void CreateSubMenuIfRequired(FToolMenuSection& Section, const FShowFlagFilter& Filter, EShowFlagGroup Group, const FName SubMenuName, const FText& MenuLabel, const FText& ToolTip) const;
+	void CreateShowFlagsSubMenu(UToolMenu* Menu, TArray<uint32> MenuCommands, int32 EntryOffset) const;
 
 	TArray<FShowFlagCommand> ShowFlagCommands;
 	bool bCommandsInitialised;
diff --git a/Engine/Source/Editor/LevelEditor/Private/SLevelViewportToolBar.cpp b/Engine/Source/Editor/LevelEditor/Private/SLevelViewportToolBar.cpp
--- a/Engine/Source/Editor/LevelEditor/Private/SLevelViewportToolBar.cpp
+++ b/Engine/Source/Editor/LevelEditor/Private/SLevelViewportToolBar.cpp
@@ -1118,7 +1118,7 @@ TSharedRef<SWidget> SLevelViewportToolBar::GenerateShowMenu() const
 	if (!UToolMenus::Get()->IsMenuRegistered(MenuName))
 	{
 		UToolMenu* Menu = UToolMenus::Get()->RegisterMenu(MenuName);
-		Menu->AddDynamicSection("DynamicSection", FNewToolMenuDelegate::CreateLambda([](UToolMenu* InMenu)
+		Menu->AddDynamicSection("LevelDynamicSection", FNewToolMenuDelegate::CreateLambda([](UToolMenu* InMenu)
 		{
 			ULevelViewportToolBarContext* Context = InMenu->FindContext<ULevelViewportToolBarContext>();
 			Context->LevelViewportToolBarWidgetConst.Pin()->FillShowMenu(InMenu);
@@ -1149,10 +1149,7 @@ void SLevelViewportToolBar::FillShowMenu(UToolMenu* Menu) const
 			Section.AddMenuEntry(Actions.UseDefaultShowFlags);
 		}
 
-		Menu->AddDynamicSection("ShowFlagsMenu", FNewToolMenuDelegateLegacy::CreateLambda([](FMenuBuilder& InMenuBuilder, UToolMenu* InMenu)
-		{
-			FShowFlagMenuCommands::Get().BuildShowFlagsMenu(InMenuBuilder);
-		}));
+		FShowFlagMenuCommands::Get().BuildShowFlagsMenu(Menu);
 
 		FText ShowAllLabel = LOCTEXT("ShowAllLabel", "Show All");
 		FText HideAllLabel = LOCTEXT("HideAllLabel", "Hide All");
diff --git a/Engine/Source/Editor/Persona/Persona.Build.cs b/Engine/Source/Editor/Persona/Persona.Build.cs
--- a/Engine/Source/Editor/Persona/Persona.Build.cs
+++ b/Engine/Source/Editor/Persona/Persona.Build.cs
@@ -67,6 +67,7 @@ public class Persona : ModuleRules
                 "PinnedCommandList",
                 "RenderCore",
 				"SkeletalMeshUtilitiesCommon",
+				"ToolMenus",
                 "CurveEditor",
 				"SequencerWidgets",
 				"TimeManagement",
diff --git a/Engine/Source/Editor/Persona/Private/SAnimViewportToolBar.cpp b/Engine/Source/Editor/Persona/Private/SAnimViewportToolBar.cpp
--- a/Engine/Source/Editor/Persona/Private/SAnimViewportToolBar.cpp
+++ b/Engine/Source/Editor/Persona/Private/SAnimViewportToolBar.cpp
@@ -5,6 +5,7 @@
 #include "Widgets/Layout/SBorder.h"
 #include "Widgets/Text/STextBlock.h"
 #include "Framework/MultiBox/MultiBoxBuilder.h"
+#include "ToolMenus.h"
 #include "EngineGlobals.h"
 #include "AssetData.h"
 #include "Engine/Engine.h"
@@ -808,54 +809,60 @@ void SAnimViewportToolBar::FillCharacterClothingMenu(FMenuBuilder& MenuBuilder)
 
 TSharedRef<SWidget> SAnimViewportToolBar::GenerateShowMenu() const
 {
-	TSharedPtr<FExtender> MenuExtender = FExtender::Combine(Extenders);
+	static const FName MenuName("Persona.AnimViewportToolBar");
 
-	const bool bInShouldCloseWindowAfterMenuSelection = true;
-	FMenuBuilder InMenuBuilder(bInShouldCloseWindowAfterMenuSelection, CommandList, MenuExtender);
-
-	// Only include helpful show flags.
-	static const FShowFlagFilter ShowFlagFilter = FShowFlagFilter(FShowFlagFilter::ExcludeAllFlagsByDefault)
-		// General
-		.IncludeFlag(FEngineShowFlags::SF_AntiAliasing)
-		.IncludeFlag(FEngineShowFlags::SF_Collision)
-		.IncludeFlag(FEngineShowFlags::SF_Grid)
-		.IncludeFlag(FEngineShowFlags::SF_Particles)
-		.IncludeFlag(FEngineShowFlags::SF_Translucency)
-		// Post Processing
-		.IncludeFlag(FEngineShowFlags::SF_Bloom)
-		.IncludeFlag(FEngineShowFlags::SF_DepthOfField)
-		.IncludeFlag(FEngineShowFlags::SF_EyeAdaptation)
-		.IncludeFlag(FEngineShowFlags::SF_HMDDistortion)
-		.IncludeFlag(FEngineShowFlags::SF_MotionBlur)
-		.IncludeFlag(FEngineShowFlags::SF_Tonemapper)
-		// Lighting Components
-		.IncludeGroup(SFG_LightingComponents)
-		// Lighting Features
-		.IncludeFlag(FEngineShowFlags::SF_AmbientCubemap)
-		.IncludeFlag(FEngineShowFlags::SF_DistanceFieldAO)
-		.IncludeFlag(FEngineShowFlags::SF_IndirectLightingCache)
-		.IncludeFlag(FEngineShowFlags::SF_LightFunctions)
-		.IncludeFlag(FEngineShowFlags::SF_LightShafts)
-		.IncludeFlag(FEngineShowFlags::SF_ReflectionEnvironment)
-		.IncludeFlag(FEngineShowFlags::SF_ScreenSpaceAO)
-		.IncludeFlag(FEngineShowFlags::SF_ContactShadows)
-		.IncludeFlag(FEngineShowFlags::SF_ScreenSpaceReflections)
-		.IncludeFlag(FEngineShowFlags::SF_SubsurfaceScattering)
-		.IncludeFlag(FEngineShowFlags::SF_TexturedLightProfiles)
-		// Developer
-		.IncludeFlag(FEngineShowFlags::SF_Refraction)
-		// Advanced
-		.IncludeFlag(FEngineShowFlags::SF_DeferredLighting)
-		.IncludeFlag(FEngineShowFlags::SF_Selection)
-		.IncludeFlag(FEngineShowFlags::SF_SeparateTranslucency)
-		.IncludeFlag(FEngineShowFlags::SF_TemporalAA)
-		.IncludeFlag(FEngineShowFlags::SF_Tessellation)
-		.IncludeFlag(FEngineShowFlags::SF_VertexColors)		
-		;
-
-	FShowFlagMenuCommands::Get().BuildShowFlagsMenu(InMenuBuilder, ShowFlagFilter);
+	if (!UToolMenus::Get()->IsMenuRegistered(MenuName))
+	{
+		UToolMenu* Menu = UToolMenus::Get()->RegisterMenu(MenuName);
+		Menu->AddDynamicSection("AnimSection", FNewToolMenuDelegate::CreateLambda([](UToolMenu* InMenu)
+		{
+			// Only include helpful show flags.
+			static const FShowFlagFilter ShowFlagFilter = FShowFlagFilter(FShowFlagFilter::ExcludeAllFlagsByDefault)
+				// General
+				.IncludeFlag(FEngineShowFlags::SF_AntiAliasing)
+				.IncludeFlag(FEngineShowFlags::SF_Collision)
+				.IncludeFlag(FEngineShowFlags::SF_Grid)
+				.IncludeFlag(FEngineShowFlags::SF_Particles)
+				.IncludeFlag(FEngineShowFlags::SF_Translucency)
+				// Post Processing
+				.IncludeFlag(FEngineShowFlags::SF_Bloom)
+				.IncludeFlag(FEngineShowFlags::SF_DepthOfField)
+				.IncludeFlag(FEngineShowFlags::SF_EyeAdaptation)
+				.IncludeFlag(FEngineShowFlags::SF_HMDDistortion)
+				.IncludeFlag(FEngineShowFlags::SF_MotionBlur)
+				.IncludeFlag(FEngineShowFlags::SF_Tonemapper)
+				// Lighting Components
+				.IncludeGroup(SFG_LightingComponents)
+				// Lighting Features
+				.IncludeFlag(FEngineShowFlags::SF_AmbientCubemap)
+				.IncludeFlag(FEngineShowFlags::SF_DistanceFieldAO)
+				.IncludeFlag(FEngineShowFlags::SF_IndirectLightingCache)
+				.IncludeFlag(FEngineShowFlags::SF_LightFunctions)
+				.IncludeFlag(FEngineShowFlags::SF_LightShafts)
+				.IncludeFlag(FEngineShowFlags::SF_ReflectionEnvironment)
+				.IncludeFlag(FEngineShowFlags::SF_ScreenSpaceAO)
+				.IncludeFlag(FEngineShowFlags::SF_ContactShadows)
+				.IncludeFlag(FEngineShowFlags::SF_ScreenSpaceReflections)
+				.IncludeFlag(FEngineShowFlags::SF_SubsurfaceScattering)
+				.IncludeFlag(FEngineShowFlags::SF_TexturedLightProfiles)
+				// Developer
+				.IncludeFlag(FEngineShowFlags::SF_Refraction)
+				// Advanced
+				.IncludeFlag(FEngineShowFlags::SF_DeferredLighting)
+				.IncludeFlag(FEngineShowFlags::SF_Selection)
+				.IncludeFlag(FEngineShowFlags::SF_SeparateTranslucency)
+				.IncludeFlag(FEngineShowFlags::SF_TemporalAA)
+				.IncludeFlag(FEngineShowFlags::SF_Tessellation)
+				.IncludeFlag(FEngineShowFlags::SF_VertexColors)
+				;
+
+			FShowFlagMenuCommands::Get().BuildShowFlagsMenu(InMenu, ShowFlagFilter);
+		}));
+	}
 
-	return InMenuBuilder.MakeWidget();
+	TSharedPtr<FExtender> MenuExtender = FExtender::Combine(Extenders);
+	FToolMenuContext MenuContext(CommandList, MenuExtender);
+	return UToolMenus::Get()->GenerateWidget(MenuName, MenuContext);
 }
 
 FText SAnimViewportToolBar::GetLODMenuLabel() const
