Additional fixes to go with 22968432
#rnx

diff --git a/Engine/Source/Editor/CommonMenuExtensions/Private/StrataVisualizationMenuCommands.cpp b/Engine/Source/Editor/CommonMenuExtensions/Private/StrataVisualizationMenuCommands.cpp
--- a/Engine/Source/Editor/CommonMenuExtensions/Private/StrataVisualizationMenuCommands.cpp
+++ b/Engine/Source/Editor/CommonMenuExtensions/Private/StrataVisualizationMenuCommands.cpp
@@ -119,23 +119,29 @@ void FStrataVisualizationMenuCommands::BindCommands(FUICommandList& CommandList,
 		const FStrataVisualizationMenuCommands::FStrataVisualizationRecord& Record = It.Value();
 		CommandList.MapAction(
 			Record.Command,
-			FExecuteAction::CreateStatic<const TSharedPtr<FEditorViewportClient>&>(&FStrataVisualizationMenuCommands::ChangeStrataVisualizationMode, Client, Record.Name),
+			FExecuteAction::CreateStatic(&FStrataVisualizationMenuCommands::ChangeStrataVisualizationMode, Client.ToWeakPtr(), Record.Name),
 			FCanExecuteAction(),
-			FIsActionChecked::CreateStatic<const TSharedPtr<FEditorViewportClient>&>(&FStrataVisualizationMenuCommands::IsStrataVisualizationModeSelected, Client, Record.Name)
+			FIsActionChecked::CreateStatic(&FStrataVisualizationMenuCommands::IsStrataVisualizationModeSelected, Client.ToWeakPtr(), Record.Name)
 		);
 	}
 }
 
-void FStrataVisualizationMenuCommands::ChangeStrataVisualizationMode(const TSharedPtr<FEditorViewportClient>& Client, FName InName)
+void FStrataVisualizationMenuCommands::ChangeStrataVisualizationMode(TWeakPtr<FEditorViewportClient> WeakClient, FName InName)
 {
-	check(Client.IsValid());
-	Client->ChangeStrataVisualizationMode(InName);
+	if (TSharedPtr<FEditorViewportClient> Client = WeakClient.Pin())
+	{
+		Client->ChangeStrataVisualizationMode(InName);
+	}
 }
 
-bool FStrataVisualizationMenuCommands::IsStrataVisualizationModeSelected(const TSharedPtr<FEditorViewportClient>& Client, FName InName)
+bool FStrataVisualizationMenuCommands::IsStrataVisualizationModeSelected(TWeakPtr<FEditorViewportClient> WeakClient, FName InName)
 {
-	check(Client.IsValid());
-	return Client->IsStrataVisualizationModeSelected(InName);
+	if (TSharedPtr<FEditorViewportClient> Client = WeakClient.Pin())
+	{
+		return Client->IsStrataVisualizationModeSelected(InName);
+	}
+	
+	return false;
 }
 
 #undef LOCTEXT_NAMESPACE
diff --git a/Engine/Source/Editor/CommonMenuExtensions/Public/StrataVisualizationMenuCommands.h b/Engine/Source/Editor/CommonMenuExtensions/Public/StrataVisualizationMenuCommands.h
--- a/Engine/Source/Editor/CommonMenuExtensions/Public/StrataVisualizationMenuCommands.h
+++ b/Engine/Source/Editor/CommonMenuExtensions/Public/StrataVisualizationMenuCommands.h
@@ -54,8 +54,8 @@ private:
 	void BuildCommandMap();
 	bool AddCommandTypeToMenu(FMenuBuilder& Menu, const FStrataVisualizationType Type) const;
 
-	static void ChangeStrataVisualizationMode(const TSharedPtr<FEditorViewportClient>& Client, FName InName);
-	static bool IsStrataVisualizationModeSelected(const TSharedPtr<FEditorViewportClient>& Client, FName InName);
+	static void ChangeStrataVisualizationMode(TWeakPtr<FEditorViewportClient> WeakClient, FName InName);
+	static bool IsStrataVisualizationModeSelected(TWeakPtr<FEditorViewportClient> WeakClient, FName InName);
 
 private:
 	TStrataVisualizationModeCommandMap CommandMap;
