Fix CIS issue.

diff --git a/Engine/Plugins/FX/Niagara/Source/Niagara/Private/NiagaraDataInterfaceRenderTargetVolume.cpp b/Engine/Plugins/FX/Niagara/Source/Niagara/Private/NiagaraDataInterfaceRenderTargetVolume.cpp
--- a/Engine/Plugins/FX/Niagara/Source/Niagara/Private/NiagaraDataInterfaceRenderTargetVolume.cpp
+++ b/Engine/Plugins/FX/Niagara/Source/Niagara/Private/NiagaraDataInterfaceRenderTargetVolume.cpp
@@ -5,6 +5,7 @@
 #include "TextureResource.h"
 #include "TextureRenderTargetVolumeResource.h"
 #include "Engine/TextureRenderTargetVolume.h"
+#include "HAL/PlatformFileManager.h"
 
 #include "NDIRenderTargetVolumeSimCacheData.h"
 #include "NiagaraDataInterfaceRenderTargetCommon.h"
diff --git a/Engine/Plugins/Runtime/HairStrands/Source/HairStrandsCore/Private/Deformer/DeformerDataInterfaceGroom.cpp b/Engine/Plugins/Runtime/HairStrands/Source/HairStrandsCore/Private/Deformer/DeformerDataInterfaceGroom.cpp
--- a/Engine/Plugins/Runtime/HairStrands/Source/HairStrandsCore/Private/Deformer/DeformerDataInterfaceGroom.cpp
+++ b/Engine/Plugins/Runtime/HairStrands/Source/HairStrandsCore/Private/Deformer/DeformerDataInterfaceGroom.cpp
@@ -10,6 +10,7 @@
 #include "GroomComponent.h"
 #include "GroomInstance.h"
 #include "RenderGraphBuilder.h"
+#include "GlobalRenderResources.h"
 
 FString UOptimusGroomDataInterface::GetDisplayName() const
 {
diff --git a/Engine/Source/Editor/CommonMenuExtensions/Private/GroomVisualizationMenuCommands.cpp b/Engine/Source/Editor/CommonMenuExtensions/Private/GroomVisualizationMenuCommands.cpp
--- a/Engine/Source/Editor/CommonMenuExtensions/Private/GroomVisualizationMenuCommands.cpp
+++ b/Engine/Source/Editor/CommonMenuExtensions/Private/GroomVisualizationMenuCommands.cpp
@@ -74,7 +74,7 @@ void FGroomVisualizationMenuCommands::InternalBuildVisualisationSubMenu(FMenuBui
 
 		// Attributes
 		{
-			Menu.BeginSection("LevelViewportGroomVisualizationGeneral", LOCTEXT("GroomVisualizationGeneral", "Groom Attributes"));
+			Menu.BeginSection("LevelViewportGroomVisualizationAttributes", LOCTEXT("GroomVisualizationAttributes", "Groom Attributes"));
 			Commands.AddCommandTypeToMenu(Menu, EGroomViewMode::RenderHairSeed);
 			Commands.AddCommandTypeToMenu(Menu, EGroomViewMode::RenderHairRootUV);
 			Commands.AddCommandTypeToMenu(Menu, EGroomViewMode::RenderHairRootUDIM);
@@ -87,7 +87,7 @@ void FGroomVisualizationMenuCommands::InternalBuildVisualisationSubMenu(FMenuBui
 		}
 		// Guides
 		{
-			Menu.BeginSection("LevelViewportGroomVisualizationGeneral", LOCTEXT("GroomVisualizationGuides", "Groom Guides"));
+			Menu.BeginSection("LevelViewportGroomVisualizationGuides", LOCTEXT("GroomVisualizationGuides", "Groom Guides"));
 			Commands.AddCommandTypeToMenu(Menu, EGroomViewMode::SimHairStrands);
 			Commands.AddCommandTypeToMenu(Menu, EGroomViewMode::RenderHairStrands);
 			Commands.AddCommandTypeToMenu(Menu, EGroomViewMode::CardGuides);
@@ -96,7 +96,7 @@ void FGroomVisualizationMenuCommands::InternalBuildVisualisationSubMenu(FMenuBui
 
 		// Strands
 		{
-			Menu.BeginSection("LevelViewportGroomVisualizationGeneral", LOCTEXT("GroomVisualizationMaterial", "Groom Strands"));
+			Menu.BeginSection("LevelViewportGroomVisualizationStrands", LOCTEXT("GroomVisualizationMaterial", "Groom Strands"));
 			Commands.AddCommandTypeToMenu(Menu, EGroomViewMode::RenderHairControlPoints);
 			Commands.AddCommandTypeToMenu(Menu, EGroomViewMode::MaterialDepth);
 			Commands.AddCommandTypeToMenu(Menu, EGroomViewMode::MaterialBaseColor);
