Fix CIS issue.

diff --git a/Engine/Source/Editor/CommonMenuExtensions/Private/StrataVisualizationMenuCommands.cpp b/Engine/Source/Editor/CommonMenuExtensions/Private/StrataVisualizationMenuCommands.cpp
--- a/Engine/Source/Editor/CommonMenuExtensions/Private/StrataVisualizationMenuCommands.cpp
+++ b/Engine/Source/Editor/CommonMenuExtensions/Private/StrataVisualizationMenuCommands.cpp
@@ -60,7 +60,7 @@ void FStrataVisualizationMenuCommands::BuildVisualisationSubMenu(FMenuBuilder& M
 	{
 		// General
 		{
-			Menu.BeginSection("LevelViewportSubstrateVisualizationMode", LOCTEXT("SubstrateVisualizationHeader", "Substrate General View Mode"));
+			Menu.BeginSection("LevelViewportSubstrateVisualizationModeGeneral", LOCTEXT("SubstrateVisualizationGeneral", "Substrate General View Mode"));
 			Commands.AddCommandTypeToMenu(Menu, FStrataViewMode::MaterialProperties);
 			Commands.AddCommandTypeToMenu(Menu, FStrataViewMode::MaterialCount);
 			Commands.AddCommandTypeToMenu(Menu, FStrataViewMode::MaterialByteCount);
@@ -70,7 +70,7 @@ void FStrataVisualizationMenuCommands::BuildVisualisationSubMenu(FMenuBuilder& M
 
 		// Advanced
 		{
-			Menu.BeginSection("LevelViewportSubstrateVisualizationMode", LOCTEXT("SubstrateVisualizationHeader", "Substrate Advanced View Mode"));
+			Menu.BeginSection("LevelViewportSubstrateVisualizationModeAdvanced", LOCTEXT("SubstrateVisualizationAdvanced", "Substrate Advanced View Mode"));
 			Commands.AddCommandTypeToMenu(Menu, FStrataViewMode::AdvancedMaterialProperties);
 			Commands.AddCommandTypeToMenu(Menu, FStrataViewMode::MaterialClassification);
 			Commands.AddCommandTypeToMenu(Menu, FStrataViewMode::RoughRefractionClassification);
