Extract functions for building the Common Show Flags and All Show Flags sections of the viewport toolbar "Show" menu

We need this for sharing code with the new level editor viewport toolbar

diff --git a/Engine/Source/Editor/CommonMenuExtensions/Private/ShowFlagMenuCommands.cpp b/Engine/Source/Editor/CommonMenuExtensions/Private/ShowFlagMenuCommands.cpp
--- a/Engine/Source/Editor/CommonMenuExtensions/Private/ShowFlagMenuCommands.cpp
+++ b/Engine/Source/Editor/CommonMenuExtensions/Private/ShowFlagMenuCommands.cpp
@@ -79,24 +79,21 @@ void FShowFlagMenuCommands::BuildShowFlagsMenu(UToolMenu* Menu, const FShowFlagF
 		return;
 	}
 
-	CreateCommonShowFlagMenuItems(Menu, Filter);
+	{
+		FToolMenuSection& Section =
+			Menu->AddSection("ShowFlagsMenuSectionCommon", LOCTEXT("CommonShowFlagHeader", "Common Show Flags"));
+
+		PopulateCommonShowFlagsSection(Section, Filter);
+	}
 
 	{
 		FToolMenuSection& Section = Menu->AddSection("LevelViewportShowFlags", LOCTEXT("AllShowFlagHeader", "All Show Flags"));
-		CreateSubMenuIfRequired(Section, Filter, SFG_PostProcess, "SFG_PostProcess", LOCTEXT("PostProcessShowFlagsMenu", "Post Processing"), LOCTEXT("PostProcessShowFlagsMenu_ToolTip", "Post process show flags"), "ShowFlagsMenu.SubMenu.PostProcessing");
-		CreateSubMenuIfRequired(Section, Filter, SFG_LightTypes, "SFG_LightTypes", LOCTEXT("LightTypesShowFlagsMenu", "Light Types"), LOCTEXT("LightTypesShowFlagsMenu_ToolTip", "Light Types show flags"), "ShowFlagsMenu.SubMenu.LightTypes");
-		CreateSubMenuIfRequired(Section, Filter, SFG_LightingComponents, "SFG_LightingComponents", LOCTEXT("LightingComponentsShowFlagsMenu", "Lighting Components"), LOCTEXT("LightingComponentsShowFlagsMenu_ToolTip", "Lighting Components show flags"), "ShowFlagsMenu.SubMenu.LightingComponents");
-		CreateSubMenuIfRequired(Section, Filter, SFG_LightingFeatures, "SFG_LightingFeatures", LOCTEXT("LightingFeaturesShowFlagsMenu", "Lighting Features"), LOCTEXT("LightingFeaturesShowFlagsMenu_ToolTip", "Lighting Features show flags"), "ShowFlagsMenu.SubMenu.LightingFeatures");
-		CreateSubMenuIfRequired(Section, Filter, SFG_Lumen, "SFG_Lumen", LOCTEXT("LumenShowFlagsMenu", "Lumen"), LOCTEXT("LumenShowFlagsMenu_ToolTip", "Lumen show flags"), "ShowFlagsMenu.SubMenu.Lumen");
-		CreateSubMenuIfRequired(Section, Filter, SFG_Nanite, "SFG_Nanite", LOCTEXT("NaniteShowFlagsMenu", "Nanite"), LOCTEXT("NaniteShowFlagsMenu_ToolTip", "Nanite show flags"), "ShowFlagsMenu.SubMenu.Nanite");
-		CreateSubMenuIfRequired(Section, Filter, SFG_Developer, "SFG_Developer", LOCTEXT("DeveloperShowFlagsMenu", "Developer"), LOCTEXT("DeveloperShowFlagsMenu_ToolTip", "Developer show flags"), "ShowFlagsMenu.SubMenu.Developer");
-		CreateSubMenuIfRequired(Section, Filter, SFG_Visualize, "SFG_Visualize", LOCTEXT("VisualizeShowFlagsMenu", "Visualize"), LOCTEXT("VisualizeShowFlagsMenu_ToolTip", "Visualize show flags"), "ShowFlagsMenu.SubMenu.Visualize");
-		CreateSubMenuIfRequired(Section, Filter, SFG_Advanced, "SFG_Advanced", LOCTEXT("AdvancedShowFlagsMenu", "Advanced"), LOCTEXT("AdvancedShowFlagsMenu_ToolTip", "Advanced show flags"), "ShowFlagsMenu.SubMenu.Advanced");
-		CreateSubMenuIfRequired(Section, Filter, SFG_Custom, "SFG_Custom", LOCTEXT("CustomShowFlagsMenu", "Custom"), LOCTEXT("CustomShowFlagsMenu_ToolTip", "Custom show flags"), NAME_None);
+
+		PopulateAllShowFlagsSection(Section, Filter);
 	}
 }
 
-void FShowFlagMenuCommands::CreateCommonShowFlagMenuItems(UToolMenu* Menu, const FShowFlagFilter& Filter) const
+void FShowFlagMenuCommands::PopulateCommonShowFlagsSection(FToolMenuSection& Section, const FShowFlagFilter& Filter) const
 {
 	const FShowFlagFilter::FGroupedShowFlagIndices& GroupedFlagIndices = Filter.GetFilteredIndices();
 	const TArray<uint32>& FlagIndices = GroupedFlagIndices[SFG_Normal];
@@ -106,19 +103,114 @@ void FShowFlagMenuCommands::CreateCommonShowFlagMenuItems(UToolMenu* Menu, const
 		return;
 	}
 
+	for (int32 ArrayIndex = 0; ArrayIndex < FlagIndices.Num(); ++ArrayIndex)
 	{
-		FToolMenuSection& Section = Menu->AddSection("ShowFlagsMenuSectionCommon", LOCTEXT("CommonShowFlagHeader", "Common Show Flags"));
-		for (int32 ArrayIndex = 0; ArrayIndex < FlagIndices.Num(); ++ArrayIndex)
-		{
-			const uint32 FlagIndex = FlagIndices[ArrayIndex];
-			const FShowFlagCommand& ShowFlagCommand = ShowFlagCommands[FlagIndex];
+		const uint32 FlagIndex = FlagIndices[ArrayIndex];
+		const FShowFlagCommand& ShowFlagCommand = ShowFlagCommands[FlagIndex];
 
-			ensure(Section.FindEntry(ShowFlagCommand.ShowMenuItem->GetCommandName()) == nullptr);
-			Section.AddMenuEntry(*FString::Printf(TEXT("Common_%s"), *ShowFlagCommand.ShowMenuItem->GetCommandName().ToString()), ShowFlagCommand.ShowMenuItem, ShowFlagCommand.LabelOverride);
-		}
+		ensure(Section.FindEntry(ShowFlagCommand.ShowMenuItem->GetCommandName()) == nullptr);
+		Section.AddMenuEntry(
+			*FString::Printf(TEXT("Common_%s"), *ShowFlagCommand.ShowMenuItem->GetCommandName().ToString()),
+			ShowFlagCommand.ShowMenuItem,
+			ShowFlagCommand.LabelOverride
+		);
 	}
 }
 
+void FShowFlagMenuCommands::PopulateAllShowFlagsSection(FToolMenuSection& Section, const FShowFlagFilter& Filter) const
+{
+	CreateSubMenuIfRequired(
+		Section,
+		Filter,
+		SFG_PostProcess,
+		"SFG_PostProcess",
+		LOCTEXT("PostProcessShowFlagsMenu", "Post Processing"),
+		LOCTEXT("PostProcessShowFlagsMenu_ToolTip", "Post process show flags"),
+		"ShowFlagsMenu.SubMenu.PostProcessing"
+	);
+	CreateSubMenuIfRequired(
+		Section,
+		Filter,
+		SFG_LightTypes,
+		"SFG_LightTypes",
+		LOCTEXT("LightTypesShowFlagsMenu", "Light Types"),
+		LOCTEXT("LightTypesShowFlagsMenu_ToolTip", "Light Types show flags"),
+		"ShowFlagsMenu.SubMenu.LightTypes"
+	);
+	CreateSubMenuIfRequired(
+		Section,
+		Filter,
+		SFG_LightingComponents,
+		"SFG_LightingComponents",
+		LOCTEXT("LightingComponentsShowFlagsMenu", "Lighting Components"),
+		LOCTEXT("LightingComponentsShowFlagsMenu_ToolTip", "Lighting Components show flags"),
+		"ShowFlagsMenu.SubMenu.LightingComponents"
+	);
+	CreateSubMenuIfRequired(
+		Section,
+		Filter,
+		SFG_LightingFeatures,
+		"SFG_LightingFeatures",
+		LOCTEXT("LightingFeaturesShowFlagsMenu", "Lighting Features"),
+		LOCTEXT("LightingFeaturesShowFlagsMenu_ToolTip", "Lighting Features show flags"),
+		"ShowFlagsMenu.SubMenu.LightingFeatures"
+	);
+	CreateSubMenuIfRequired(
+		Section,
+		Filter,
+		SFG_Lumen,
+		"SFG_Lumen",
+		LOCTEXT("LumenShowFlagsMenu", "Lumen"),
+		LOCTEXT("LumenShowFlagsMenu_ToolTip", "Lumen show flags"),
+		"ShowFlagsMenu.SubMenu.Lumen"
+	);
+	CreateSubMenuIfRequired(
+		Section,
+		Filter,
+		SFG_Nanite,
+		"SFG_Nanite",
+		LOCTEXT("NaniteShowFlagsMenu", "Nanite"),
+		LOCTEXT("NaniteShowFlagsMenu_ToolTip", "Nanite show flags"),
+		"ShowFlagsMenu.SubMenu.Nanite"
+	);
+	CreateSubMenuIfRequired(
+		Section,
+		Filter,
+		SFG_Developer,
+		"SFG_Developer",
+		LOCTEXT("DeveloperShowFlagsMenu", "Developer"),
+		LOCTEXT("DeveloperShowFlagsMenu_ToolTip", "Developer show flags"),
+		"ShowFlagsMenu.SubMenu.Developer"
+	);
+	CreateSubMenuIfRequired(
+		Section,
+		Filter,
+		SFG_Visualize,
+		"SFG_Visualize",
+		LOCTEXT("VisualizeShowFlagsMenu", "Visualize"),
+		LOCTEXT("VisualizeShowFlagsMenu_ToolTip", "Visualize show flags"),
+		"ShowFlagsMenu.SubMenu.Visualize"
+	);
+	CreateSubMenuIfRequired(
+		Section,
+		Filter,
+		SFG_Advanced,
+		"SFG_Advanced",
+		LOCTEXT("AdvancedShowFlagsMenu", "Advanced"),
+		LOCTEXT("AdvancedShowFlagsMenu_ToolTip", "Advanced show flags"),
+		"ShowFlagsMenu.SubMenu.Advanced"
+	);
+	CreateSubMenuIfRequired(
+		Section,
+		Filter,
+		SFG_Custom,
+		"SFG_Custom",
+		LOCTEXT("CustomShowFlagsMenu", "Custom"),
+		LOCTEXT("CustomShowFlagsMenu_ToolTip", "Custom show flags"),
+		NAME_None
+	);
+}
+
 void FShowFlagMenuCommands::CreateSubMenuIfRequired(FToolMenuSection& Section, const FShowFlagFilter& Filter, EShowFlagGroup Group, const FName SubMenuName, const FText& MenuLabel, const FText& ToolTip, const FName IconName) const
 {
 	const FShowFlagFilter::FGroupedShowFlagIndices& GroupedFlagIndices = Filter.GetFilteredIndices();
diff --git a/Engine/Source/Editor/CommonMenuExtensions/Public/ShowFlagMenuCommands.h b/Engine/Source/Editor/CommonMenuExtensions/Public/ShowFlagMenuCommands.h
--- a/Engine/Source/Editor/CommonMenuExtensions/Public/ShowFlagMenuCommands.h
+++ b/Engine/Source/Editor/CommonMenuExtensions/Public/ShowFlagMenuCommands.h
@@ -53,6 +53,13 @@ public:
 	void BindCommands(FUICommandList& CommandList, const TSharedPtr<FEditorViewportClient>& Client) const;
 	void BuildShowFlagsMenu(UToolMenu* Menu, const FShowFlagFilter& Filter = FShowFlagFilter(FShowFlagFilter::IncludeAllFlagsByDefault)) const;
 
+	void PopulateCommonShowFlagsSection(
+		FToolMenuSection& Section, const FShowFlagFilter& Filter = FShowFlagFilter(FShowFlagFilter::IncludeAllFlagsByDefault)
+	) const;
+	void PopulateAllShowFlagsSection(
+		FToolMenuSection& Section, const FShowFlagFilter& Filter = FShowFlagFilter(FShowFlagFilter::IncludeAllFlagsByDefault)
+	) const;
+
 private:
 	static void StaticCreateShowFlagsSubMenu(UToolMenu* Menu, TArray<uint32> FlagIndices, int32 EntryOffset);
 	static void ToggleShowFlag(TWeakPtr<FEditorViewportClient> WeakClient, FEngineShowFlags::EShowFlag EngineShowFlagIndex);
@@ -62,7 +69,6 @@ private:
 
 	void CreateShowFlagCommands();
 	void UpdateCustomShowFlags() const;
-	void CreateCommonShowFlagMenuItems(UToolMenu* Menu, const FShowFlagFilter& Filter) const;
 	void CreateSubMenuIfRequired(FToolMenuSection& Section, const FShowFlagFilter& Filter, EShowFlagGroup Group, const FName SubMenuName, const FText& MenuLabel, const FText& ToolTip, const FName IconName) const;
 	void CreateShowFlagsSubMenu(UToolMenu* Menu, TArray<uint32> MenuCommands, int32 EntryOffset) const;
 
