[VSM] Add visualization menu item in editor for ShowStats.

diff --git a/Engine/Source/Editor/CommonMenuExtensions/Private/VirtualShadowMapVisualizationMenuCommands.cpp b/Engine/Source/Editor/CommonMenuExtensions/Private/VirtualShadowMapVisualizationMenuCommands.cpp
--- a/Engine/Source/Editor/CommonMenuExtensions/Private/VirtualShadowMapVisualizationMenuCommands.cpp
+++ b/Engine/Source/Editor/CommonMenuExtensions/Private/VirtualShadowMapVisualizationMenuCommands.cpp
@@ -70,6 +70,7 @@ void FVirtualShadowMapVisualizationMenuCommands::BuildCommandMap()
 		}
 	}
 
+	UI_COMMAND(ShowStatsCommand, "Show Statistics", "r.Shadow.Virtual.ShowStats", EUserInterfaceActionType::Check, FInputChord());
 	UI_COMMAND(VisualizeNextLightCommand, "Visualize next light", "r.Shadow.Virtual.Visualize.NextLight", EUserInterfaceActionType::Button, FInputChord());
 	UI_COMMAND(VisualizePrevLightCommand, "Visualize previous light", "r.Shadow.Virtual.Visualize.PrevLight", EUserInterfaceActionType::Button, FInputChord());
 }
@@ -97,6 +98,8 @@ bool FVirtualShadowMapVisualizationMenuCommands::AddCommandTypeToMenu(FMenuBuild
 {
 	bool bAddedCommands = false;
 
+	Menu.AddMenuEntry(ShowStatsCommand, NAME_None, ShowStatsCommand->GetLabel());
+
 	const TVirtualShadowMapVisualizationModeCommandMap& Commands = CommandMap;
 	for (TCommandConstIterator It = CreateCommandConstIterator(); It; ++It)
 	{
@@ -139,6 +142,18 @@ void FVirtualShadowMapVisualizationMenuCommands::BindCommands(FUICommandList& Co
 		);
 	}
 
+	IConsoleVariable* ShowStatsCVar = IConsoleManager::Get().FindConsoleVariable(TEXT("r.Shadow.Virtual.ShowStats"));
+	check(ShowStatsCVar);
+	CommandList.MapAction(
+		ShowStatsCommand,
+		FExecuteAction::CreateLambda([=]()
+		{
+			bool bStatsVisible = ShowStatsCVar->GetInt() != 0;
+			ShowStatsCVar->Set(bStatsVisible ? 0 : 1);
+		}),
+		FCanExecuteAction::CreateLambda([]() { return true; }),
+		FIsActionChecked::CreateLambda([=]() { return ShowStatsCVar->GetInt() != 0; })
+	);
 	CommandList.MapAction(
 		VisualizeNextLightCommand,
 		FExecuteAction::CreateLambda([]()
diff --git a/Engine/Source/Editor/CommonMenuExtensions/Public/VirtualShadowMapVisualizationMenuCommands.h b/Engine/Source/Editor/CommonMenuExtensions/Public/VirtualShadowMapVisualizationMenuCommands.h
--- a/Engine/Source/Editor/CommonMenuExtensions/Public/VirtualShadowMapVisualizationMenuCommands.h
+++ b/Engine/Source/Editor/CommonMenuExtensions/Public/VirtualShadowMapVisualizationMenuCommands.h
@@ -66,6 +66,7 @@ private:
 
 private:
 	TVirtualShadowMapVisualizationModeCommandMap CommandMap;
+	TSharedPtr<FUICommandInfo> ShowStatsCommand;
 	TSharedPtr<FUICommandInfo> VisualizeNextLightCommand;
 	TSharedPtr<FUICommandInfo> VisualizePrevLightCommand;
 };
