[HWRT / Editor] Moved FRayTracingDebugVisualizationMenuCommands to CommonMenuExtensions
- and changed how submenu is registered to match how majority of other rendering visualization modes are handled.
- explicitly add the ray tracing debug visualization submenu to SStaticMeshEditorViewport.cpp via a new UE::StaticMeshEditor::ExtendViewModesSubmenu (similar to what we do in SAnimationEditorViewport.cpp).

#Synced-CL 50401438

diff --git a/Engine/Plugins/Runtime/nDisplay/Source/DisplayClusterLightCardEditor/Private/Viewport/DisplayClusterLightCardEditorViewportClient.cpp b/Engine/Plugins/Runtime/nDisplay/Source/DisplayClusterLightCardEditor/Private/Viewport/DisplayClusterLightCardEditorViewportClient.cpp
--- a/Engine/Plugins/Runtime/nDisplay/Source/DisplayClusterLightCardEditor/Private/Viewport/DisplayClusterLightCardEditorViewportClient.cpp
+++ b/Engine/Plugins/Runtime/nDisplay/Source/DisplayClusterLightCardEditor/Private/Viewport/DisplayClusterLightCardEditorViewportClient.cpp
@@ -53,7 +53,6 @@
 #include "Misc/ScopeLock.h"
 #include "PreviewScene.h"
 #include "ProceduralMeshComponent.h"
-#include "RayTracingDebugVisualizationMenuCommands.h"
 #include "RenderResource.h"
 #include "ScopedTransaction.h"
 #include "Settings/LevelEditorViewportSettings.h"
diff --git a/Engine/Source/Editor/CommonMenuExtensions/Private/CommonMenuExtensionsModule.cpp b/Engine/Source/Editor/CommonMenuExtensions/Private/CommonMenuExtensionsModule.cpp
--- a/Engine/Source/Editor/CommonMenuExtensions/Private/CommonMenuExtensionsModule.cpp
+++ b/Engine/Source/Editor/CommonMenuExtensions/Private/CommonMenuExtensionsModule.cpp
@@ -9,6 +9,7 @@
 #include "GroomVisualizationMenuCommands.h"
 #include "VirtualShadowMapVisualizationMenuCommands.h"
 #include "VirtualTextureVisualizationMenuCommands.h"
+#include "RayTracingDebugVisualizationMenuCommands.h"
 #include "ShowFlagMenuCommands.h"
 #include "Modules/ModuleManager.h"
 
@@ -18,6 +19,7 @@ void FCommonMenuExtensionsModule::StartupModule()
 {
 	FBufferVisualizationMenuCommands::Register();
 	FNaniteVisualizationMenuCommands::Register();
+	FRayTracingDebugVisualizationMenuCommands::Register();
 	FLumenVisualizationMenuCommands::Register();
 	FMegaLightsVisualizationMenuCommands::Register();
 	FSubstrateVisualizationMenuCommands::Register();
@@ -37,5 +39,6 @@ void FCommonMenuExtensionsModule::ShutdownModule()
 	FSubstrateVisualizationMenuCommands::Unregister();
 	FLumenVisualizationMenuCommands::Unregister();
 	FMegaLightsVisualizationMenuCommands::Unregister();
+	FRayTracingDebugVisualizationMenuCommands::Unregister();
 	FBufferVisualizationMenuCommands::Unregister();
-}
\ No newline at end of file
+}
diff --git a/Engine/Source/Editor/CommonMenuExtensions/Private/RayTracingDebugVisualizationMenuCommands.cpp b/Engine/Source/Editor/CommonMenuExtensions/Private/RayTracingDebugVisualizationMenuCommands.cpp
new file mode 100644
--- /dev/null
+++ b/Engine/Source/Editor/CommonMenuExtensions/Private/RayTracingDebugVisualizationMenuCommands.cpp
@@ -0,0 +1,205 @@
+// Copyright Epic Games, Inc. All Rights Reserved.
+
+#include "RayTracingDebugVisualizationMenuCommands.h"
+
+#include "DataDrivenShaderPlatformInfo.h"
+#include "Containers/UnrealString.h"
+#include "Framework/Commands/InputChord.h"
+#include "Internationalization/Text.h"
+#include "Templates/Function.h"
+#include "RayTracingVisualizationData.h"
+#include "Styling/AppStyle.h"
+#include "EditorViewportClient.h"
+#include "ToolMenu.h"
+
+int32 GRayTracingVisualizeTiming = 0;
+static FAutoConsoleVariableRef CVarRayTracingVisualizeTiming(
+	TEXT("r.RayTracing.Visualize.Timing"),
+	GRayTracingVisualizeTiming,
+	TEXT("Include 'Timing' visualization modes in the in-editor 'Ray Tracing Debug' drop down menu.")
+);
+
+int32 GRayTracingVisualizeOther = 0;
+static FAutoConsoleVariableRef CVarRayTracingVisualizeOther(
+	TEXT("r.RayTracing.Visualize.Other"),
+	GRayTracingVisualizeOther,
+	TEXT("Include 'Other' visualization modes in the in-editor 'Ray Tracing Debug' drop down menu.")
+);
+
+#define LOCTEXT_NAMESPACE "RayTracingDebugVisualizationMenuCommands"
+
+UE_DEFINE_TCOMMANDS(FRayTracingDebugVisualizationMenuCommands)
+
+FRayTracingDebugVisualizationMenuCommands::FRayTracingDebugVisualizationMenuCommands()
+	: TCommands<FRayTracingDebugVisualizationMenuCommands>
+	(
+		TEXT("RayTracingDebugVisualizationMenu"), // Context name for fast lookup
+		NSLOCTEXT("Contexts", "RayTracingMenu", "Ray Tracing Debug Visualization"), // Localized context name for displaying
+		NAME_None, // Parent context name.  
+		FAppStyle::GetAppStyleSetName() // Icon Style Set
+	),
+	CommandMap()
+{
+}
+
+void FRayTracingDebugVisualizationMenuCommands::BuildCommandMap()
+{
+	const FRayTracingVisualizationData& VisualizationData = GetRayTracingVisualizationData();
+	const FRayTracingVisualizationData::TModeMap& ModeMap = VisualizationData.GetModeMap();
+
+	CommandMap.Empty();
+	for (FRayTracingVisualizationData::TModeMap::TConstIterator It = ModeMap.CreateConstIterator(); It; ++It)
+	{
+		const FRayTracingVisualizationData::FModeRecord& Entry = It.Value();
+
+		if (Entry.bHiddenInEditor)
+		{
+			continue;
+		}
+
+		FVisualizationRecord& Record = CommandMap.Add(Entry.ModeName, FVisualizationRecord());
+		Record.Name = Entry.ModeName;
+		Record.Command = FUICommandInfoDecl(
+			this->AsShared(),
+			Entry.ModeName,
+			Entry.ModeText,
+			Entry.ModeDesc)
+			.UserInterfaceType(EUserInterfaceActionType::RadioButton)
+			.DefaultChord(FInputChord()
+			);
+
+		switch (Entry.ModeType)
+		{
+		case FRayTracingVisualizationData::FModeType::Overview:
+			Record.Type = FVisualizationType::Overview;
+			break;
+
+		default:
+		case FRayTracingVisualizationData::FModeType::Standard:
+			Record.Type = FVisualizationType::Standard;
+			break;
+
+		case FRayTracingVisualizationData::FModeType::Performance:
+			Record.Type = FVisualizationType::Performance;
+			break;
+
+		case FRayTracingVisualizationData::FModeType::Timing:
+			Record.Type = FVisualizationType::Timing;
+			break;
+
+		case FRayTracingVisualizationData::FModeType::Other:
+			Record.Type = FVisualizationType::Other;
+			break;
+		}
+	}
+}
+
+void FRayTracingDebugVisualizationMenuCommands::BuildVisualisationSubMenu(UToolMenu* InMenu)
+{
+	if (!InMenu)
+	{
+		return;
+	}
+
+	const FRayTracingDebugVisualizationMenuCommands& Commands = FRayTracingDebugVisualizationMenuCommands::Get();
+	if (!Commands.IsPopulated())
+	{
+		return;
+	}
+
+	{
+		FToolMenuSection& UnnamedSection = InMenu->AddSection(NAME_None);
+		Commands.AddCommandTypeToSection(UnnamedSection, FVisualizationType::Overview);
+	}
+
+	{
+		FToolMenuSection& GeneralSection = InMenu->AddSection(
+			"LevelViewportRayTracingVisualizationGeneral", LOCTEXT("RayTracingVisualizationGeneral", "General")
+		);
+		Commands.AddCommandTypeToSection(GeneralSection, FVisualizationType::Standard);
+	}
+
+	{
+		FToolMenuSection& PerformanceSection = InMenu->AddSection(
+			"LevelViewportRayTracingVisualizationPerformance",
+			LOCTEXT("RayTracingVisualizationPerformance", "Performance")
+		);
+		Commands.AddCommandTypeToSection(PerformanceSection, FVisualizationType::Performance);
+	}
+
+	const bool bShowTiming = GRayTracingVisualizeTiming != 0;
+	if (bShowTiming)
+	{
+		FToolMenuSection& TimingSection = InMenu->AddSection("LevelViewportRayTracingVisualizationTiming", LOCTEXT("RayTracingVisualizationTiming", "Timing"));
+		Commands.AddCommandTypeToSection(TimingSection, FVisualizationType::Timing);
+	}
+
+	const bool bShowOther = GRayTracingVisualizeOther != 0;
+	if (bShowOther)
+	{
+		FToolMenuSection& OtherSection = InMenu->AddSection("LevelViewportRayTracingVisualizationOther", LOCTEXT("RayTracingVisualizationOther", "Other"));
+		Commands.AddCommandTypeToSection(OtherSection, FVisualizationType::Other);
+	}
+}
+
+bool FRayTracingDebugVisualizationMenuCommands::AddCommandTypeToSection(FToolMenuSection& InSection, const FVisualizationType Type) const
+{
+	bool bAddedCommands = false;
+
+	for (TCommandConstIterator It = CreateCommandConstIterator(); It; ++It)
+	{
+		const FVisualizationRecord& Record = It.Value();
+		if (Record.Type == Type)
+		{
+			InSection.AddMenuEntry(Record.Command, Record.Command->GetLabel());
+			bAddedCommands = true;
+		}
+	}
+
+	return bAddedCommands;
+}
+
+FRayTracingDebugVisualizationMenuCommands::TCommandConstIterator FRayTracingDebugVisualizationMenuCommands::CreateCommandConstIterator() const
+{
+	return CommandMap.CreateConstIterator();
+}
+
+void FRayTracingDebugVisualizationMenuCommands::RegisterCommands()
+{
+	BuildCommandMap();
+}
+
+void FRayTracingDebugVisualizationMenuCommands::BindCommands(FUICommandList& CommandList, const TSharedPtr<FEditorViewportClient>& Client) const
+{
+	// Map ray tracing debug visualization mode actions
+	for (FRayTracingDebugVisualizationMenuCommands::TCommandConstIterator It = FRayTracingDebugVisualizationMenuCommands::Get().CreateCommandConstIterator(); It; ++It)
+	{
+		const FRayTracingDebugVisualizationMenuCommands::FVisualizationRecord& Record = It.Value();
+		CommandList.MapAction(
+			Record.Command,
+			FExecuteAction::CreateStatic(&FRayTracingDebugVisualizationMenuCommands::ChangeRayTracingDebugVisualizationMode, Client.ToWeakPtr(), Record.Name),
+			FCanExecuteAction(),
+			FIsActionChecked::CreateStatic(&FRayTracingDebugVisualizationMenuCommands::IsRayTracingDebugVisualizationModeSelected, Client.ToWeakPtr(), Record.Name)
+		);
+	}
+}
+
+void FRayTracingDebugVisualizationMenuCommands::ChangeRayTracingDebugVisualizationMode(TWeakPtr<FEditorViewportClient> WeakClient, FName InName)
+{
+	if (TSharedPtr<FEditorViewportClient> Client = WeakClient.Pin())
+	{
+		Client->ChangeRayTracingDebugVisualizationMode(InName);
+	}
+}
+
+bool FRayTracingDebugVisualizationMenuCommands::IsRayTracingDebugVisualizationModeSelected(TWeakPtr<FEditorViewportClient> WeakClient, FName InName)
+{
+	if (TSharedPtr<FEditorViewportClient> Client = WeakClient.Pin())
+	{
+		return Client->IsRayTracingDebugVisualizationModeSelected(InName);
+	}
+
+	return false;
+}
+
+#undef LOCTEXT_NAMESPACE
diff --git a/Engine/Source/Editor/CommonMenuExtensions/Public/RayTracingDebugVisualizationMenuCommands.h b/Engine/Source/Editor/CommonMenuExtensions/Public/RayTracingDebugVisualizationMenuCommands.h
new file mode 100644
--- /dev/null
+++ b/Engine/Source/Editor/CommonMenuExtensions/Public/RayTracingDebugVisualizationMenuCommands.h
@@ -0,0 +1,75 @@
+// Copyright Epic Games, Inc. All Rights Reserved.
+
+#pragma once
+
+#include "Containers/Map.h"
+#include "UObject/NameTypes.h"
+#include "Templates/SharedPointer.h"
+#include "Framework/Commands/UICommandInfo.h"
+#include "Framework/MultiBox/MultiBoxBuilder.h"
+#include "Framework/Commands/Commands.h"
+
+#define UE_API COMMONMENUEXTENSIONS_API
+
+class FEditorViewportClient;
+class UToolMenu;
+struct FToolMenuSection;
+
+UE_DECLARE_TCOMMANDS(class FRayTracingDebugVisualizationMenuCommands, UE_API)
+
+class FRayTracingDebugVisualizationMenuCommands : public TCommands<FRayTracingDebugVisualizationMenuCommands>
+{
+public:
+	enum class FVisualizationType : uint8
+	{
+		Overview,
+		Standard,
+		Performance,
+		Timing,
+		Other,
+	};
+
+	struct FVisualizationRecord
+	{
+		FName Name;
+		TSharedPtr<FUICommandInfo> Command;
+		FVisualizationType Type;
+
+		FVisualizationRecord()
+			: Name()
+			, Command()
+			, Type(FVisualizationType::Overview)
+		{
+		}
+	};
+
+	typedef TMultiMap<FName, FVisualizationRecord> TVisualizationModeCommandMap;
+	typedef TVisualizationModeCommandMap::TConstIterator TCommandConstIterator;
+
+	UE_API FRayTracingDebugVisualizationMenuCommands();
+
+	UE_API TCommandConstIterator CreateCommandConstIterator() const;
+
+	static UE_API void BuildVisualisationSubMenu(UToolMenu* InMenu);
+
+	UE_API virtual void RegisterCommands() override;
+
+	UE_API void BindCommands(FUICommandList& CommandList, const TSharedPtr<FEditorViewportClient>& Client) const;
+
+	inline bool IsPopulated() const
+	{
+		return CommandMap.Num() > 0;
+	}
+
+private:
+	void BuildCommandMap();
+
+	bool AddCommandTypeToSection(FToolMenuSection& InSection, const FVisualizationType Type) const;
+
+	static void ChangeRayTracingDebugVisualizationMode(TWeakPtr<FEditorViewportClient> WeakClient, FName InName);
+	static bool IsRayTracingDebugVisualizationModeSelected(TWeakPtr<FEditorViewportClient> WeakClient, FName InName);
+
+	TVisualizationModeCommandMap CommandMap;
+};
+
+#undef UE_API
diff --git a/Engine/Source/Editor/LevelEditor/Private/SLevelViewport.cpp b/Engine/Source/Editor/LevelEditor/Private/SLevelViewport.cpp
--- a/Engine/Source/Editor/LevelEditor/Private/SLevelViewport.cpp
+++ b/Engine/Source/Editor/LevelEditor/Private/SLevelViewport.cpp
@@ -78,6 +78,7 @@
 #include "AssetRegistry/IAssetRegistry.h"
 #include "BufferVisualizationMenuCommands.h"
 #include "NaniteVisualizationMenuCommands.h"
+#include "RayTracingDebugVisualizationMenuCommands.h"
 #include "LumenVisualizationMenuCommands.h"
 #include "MegaLightsVisualizationMenuCommands.h"
 #include "SubstrateVisualizationMenuCommands.h"
@@ -1787,6 +1788,7 @@ void SLevelViewport::BindViewCommands( FUICommandList& OutCommandList )
 
 	FBufferVisualizationMenuCommands::Get().BindCommands(OutCommandList, Client);
 	FNaniteVisualizationMenuCommands::Get().BindCommands(OutCommandList, Client);
+	FRayTracingDebugVisualizationMenuCommands::Get().BindCommands(OutCommandList, Client);
 	FLumenVisualizationMenuCommands::Get().BindCommands(OutCommandList, Client);
 	FMegaLightsVisualizationMenuCommands::Get().BindCommands(OutCommandList, Client);
 	if (Substrate::IsSubstrateEnabled())
@@ -1960,9 +1962,6 @@ void SLevelViewport::BindPIECommands(FUICommandList& OutCommandList)
 	{
 		MAP_VIEWMODE_ACTION(Commands.PathTracingMode, VMI_PathTracing);
 		MAP_VIEWMODE_ACTION(Commands.RayTracingDebugMode, VMI_RayTracingDebug);
-
-		// const FRayTracingDebugVisualizationMenuCommands& RtDebugCommands = FRayTracingDebugVisualizationMenuCommands::Get();
-		// RtDebugCommands.BindCommands(CommandListRef, Client);
 	}
 
 	MAP_VIEWMODE_ACTION( Commands.DetailLightingMode, VMI_Lit_DetailLighting );
@@ -2628,8 +2627,6 @@ TSharedRef<SWidget> SLevelViewport::BuildPIEViewportToolbar()
 				{
 				case UE::UnrealEd::EHidableViewModeMenuSections::GPUSkinCache:
 					return false;
-				case UE::UnrealEd::EHidableViewModeMenuSections::RayTracingDebug:
-					return false;
 				default:
 					return true;
 				}
diff --git a/Engine/Source/Editor/LevelEditor/Private/ViewportToolbar/LevelEditorSubmenus.cpp b/Engine/Source/Editor/LevelEditor/Private/ViewportToolbar/LevelEditorSubmenus.cpp
--- a/Engine/Source/Editor/LevelEditor/Private/ViewportToolbar/LevelEditorSubmenus.cpp
+++ b/Engine/Source/Editor/LevelEditor/Private/ViewportToolbar/LevelEditorSubmenus.cpp
@@ -39,6 +39,7 @@
 #include "WorldPartition/IWorldPartitionEditorModule.h"
 #include "VirtualShadowMapVisualizationData.h"
 #include "Tools/EdModeInteractiveToolsContext.h"
+#include "RayTracingDebugVisualizationMenuCommands.h"
 #include "MegaLightsVisualizationData.h"
 #include "MegaLightsVisualizationMenuCommands.h"
 
@@ -1122,6 +1123,33 @@ void PopulateViewModesMenu(UToolMenu* InMenu)
 		);
 	}
 
+	if (IsRayTracingEnabled())
+	{
+		FToolMenuSection& Section = InMenu->FindOrAddSection("ViewMode");
+		Section.AddSubMenu(
+			"RayTracingDebugSubMenu",
+			LOCTEXT("RayTracingDebugSubMenu", "Ray Tracing Debug"),
+			LOCTEXT("RayTracing_ToolTip", "Select ray tracing buffer visualization view modes"),
+			FNewToolMenuDelegate::CreateStatic(&FRayTracingDebugVisualizationMenuCommands::BuildVisualisationSubMenu),
+			FUIAction(
+				FExecuteAction(),
+				FCanExecuteAction(),
+				FIsActionChecked::CreateLambda(
+					[WeakViewport = LevelViewport.ToWeakPtr()]()
+					{
+						const TSharedPtr<::SLevelViewport> Viewport = WeakViewport.Pin();
+						check(Viewport.IsValid());
+						FLevelEditorViewportClient& ViewportClient = Viewport->GetLevelViewportClient();
+						return ViewportClient.IsViewModeEnabled(VMI_RayTracingDebug);
+					}
+				)
+			),
+			EUserInterfaceActionType::RadioButton,
+			/* bInOpenSubMenuOnClick = */ false,
+			FSlateIcon(FAppStyle::GetAppStyleSetName(), "EditorViewport.RayTracingDebugMode")
+		);
+	}
+
 	{
 		FToolMenuSection& Section = InMenu->FindOrAddSection("ViewMode");
 		Section.AddSubMenu(
diff --git a/Engine/Source/Editor/StaticMeshEditor/Private/SStaticMeshEditorViewport.cpp b/Engine/Source/Editor/StaticMeshEditor/Private/SStaticMeshEditorViewport.cpp
--- a/Engine/Source/Editor/StaticMeshEditor/Private/SStaticMeshEditorViewport.cpp
+++ b/Engine/Source/Editor/StaticMeshEditor/Private/SStaticMeshEditorViewport.cpp
@@ -1195,6 +1195,8 @@ TSharedPtr<SWidget> SStaticMeshEditorViewport::BuildViewportToolbar()
 
 			// Add the "View Modes" sub menu.
 			{
+				const FName ViewModesMenuName = "StaticMeshEditor.ViewportToolbar.ViewModes";
+
 				// Stay backward-compatible with the old viewport toolbar.
 				{
 					const FName ParentSubmenuName = "UnrealEd.ViewportToolbar.View";
@@ -1205,9 +1207,11 @@ TSharedPtr<SWidget> SStaticMeshEditorViewport::BuildViewportToolbar()
 					}
 
 					// Register our ToolMenu here first, before we create the submenu, so we can set our parent.
-					UToolMenus::Get()->RegisterMenu("StaticMeshEditor.ViewportToolbar.ViewModes", ParentSubmenuName);
+					UToolMenus::Get()->RegisterMenu(ViewModesMenuName, ParentSubmenuName);
 				}
 
+				UE::StaticMeshEditor::ExtendViewModesSubmenu(ViewModesMenuName);
+
 				RightSection.AddEntry(UE::UnrealEd::CreateViewModesSubmenu());
 			}
 
diff --git a/Engine/Source/Editor/StaticMeshEditor/Private/StaticMeshEditorViewportToolbarSections.cpp b/Engine/Source/Editor/StaticMeshEditor/Private/StaticMeshEditorViewportToolbarSections.cpp
--- a/Engine/Source/Editor/StaticMeshEditor/Private/StaticMeshEditorViewportToolbarSections.cpp
+++ b/Engine/Source/Editor/StaticMeshEditor/Private/StaticMeshEditorViewportToolbarSections.cpp
@@ -5,6 +5,7 @@
 #include "SStaticMeshEditorViewport.h"
 #include "StaticMeshEditorActions.h"
 #include "StaticMeshViewportLODCommands.h"
+#include "RayTracingDebugVisualizationMenuCommands.h"
 #include "ToolMenu.h"
 #include "ToolMenuSection.h"
 #include "ToolMenus.h"
@@ -233,5 +234,50 @@ void UE::StaticMeshEditor::FillAspectRatioSubmenu(UToolMenu* InMenu)
 	}
 }
 
+void UE::StaticMeshEditor::ExtendViewModesSubmenu(FName InViewModesSubmenuName)
+{
+	UToolMenu* const Submenu = UToolMenus::Get()->ExtendMenu(InViewModesSubmenuName);
+
+	Submenu->AddDynamicSection(
+		"LevelEditorViewModesExtensionDynamicSection",
+		FNewToolMenuDelegate::CreateLambda(
+			[](UToolMenu* InDynamicMenu)
+			{
+				if (UUnrealEdViewportToolbarContext* const EditorViewportContext =
+					InDynamicMenu->FindContext<UUnrealEdViewportToolbarContext>())
+				{
+					if (TSharedPtr<SStaticMeshEditorViewport> StaticMeshEditorViewport =
+						StaticCastSharedPtr<SStaticMeshEditorViewport>(EditorViewportContext->Viewport.Pin()))
+					{
+						FToolMenuSection& Section = InDynamicMenu->FindOrAddSection("ViewMode");
+						Section.AddSubMenu(
+							"RayTracingDebugSubMenu",
+							LOCTEXT("RayTracingDebugSubMenu", "Ray Tracing Debug"),
+							LOCTEXT("RayTracing_ToolTip", "Select ray tracing buffer visualization view modes"),
+							FNewToolMenuDelegate::CreateStatic(&FRayTracingDebugVisualizationMenuCommands::BuildVisualisationSubMenu),
+							FUIAction(
+								FExecuteAction(),
+								FCanExecuteAction(),
+								FIsActionChecked::CreateLambda(
+									[WeakViewport = StaticMeshEditorViewport.ToWeakPtr()]()
+									{
+										const TSharedPtr<SStaticMeshEditorViewport> Viewport = WeakViewport.Pin();
+										check(Viewport.IsValid());
+										FEditorViewportClient& ViewportClient = Viewport->GetViewportClient();
+										return ViewportClient.IsViewModeEnabled(VMI_RayTracingDebug);
+									}
+								)
+							),
+							EUserInterfaceActionType::RadioButton,
+							/* bInOpenSubMenuOnClick = */ false,
+							FSlateIcon(FAppStyle::GetAppStyleSetName(), "EditorViewport.RayTracingDebugMode")
+						);
+					}
+				}
+			}
+		)
+	);
+}
+
 
 #undef LOCTEXT_NAMESPACE
diff --git a/Engine/Source/Editor/StaticMeshEditor/Private/StaticMeshEditorViewportToolbarSections.h b/Engine/Source/Editor/StaticMeshEditor/Private/StaticMeshEditorViewportToolbarSections.h
--- a/Engine/Source/Editor/StaticMeshEditor/Private/StaticMeshEditorViewportToolbarSections.h
+++ b/Engine/Source/Editor/StaticMeshEditor/Private/StaticMeshEditorViewportToolbarSections.h
@@ -21,4 +21,6 @@ void FillShowSubmenu(UToolMenu *InMenu);
 void FillAspectRatioSubmenu(UToolMenu* InMenu);
 
 FText GetLODMenuLabel(const TSharedPtr<SStaticMeshEditorViewport>& InStaticMeshEditorViewport);
+
+void ExtendViewModesSubmenu(FName InViewModesSubmenuName);
 }
diff --git a/Engine/Source/Editor/StaticMeshEditor/StaticMeshEditor.Build.cs b/Engine/Source/Editor/StaticMeshEditor/StaticMeshEditor.Build.cs
--- a/Engine/Source/Editor/StaticMeshEditor/StaticMeshEditor.Build.cs
+++ b/Engine/Source/Editor/StaticMeshEditor/StaticMeshEditor.Build.cs
@@ -47,6 +47,7 @@ public class StaticMeshEditor : ModuleRules
 				"EditorSubsystem",
 				"MeshDescription",
 				"StaticMeshDescription",
+				"CommonMenuExtensions",
 				"ToolMenus",
 				"DetailCustomizations",
 				"StatusBar",
diff --git a/Engine/Source/Editor/UnrealEd/Private/EditorViewportClient.cpp b/Engine/Source/Editor/UnrealEd/Private/EditorViewportClient.cpp
--- a/Engine/Source/Editor/UnrealEd/Private/EditorViewportClient.cpp
+++ b/Engine/Source/Editor/UnrealEd/Private/EditorViewportClient.cpp
@@ -57,7 +57,6 @@
 #include "Elements/Framework/TypedElementViewportInteraction.h"
 #include "ImageWriteQueue.h"
 #include "DebugViewModeHelpers.h"
-#include "RayTracingDebugVisualizationMenuCommands.h"
 #include "Misc/ScopedSlowTask.h"
 #include "UnrealEngine.h"
 #include "BufferVisualizationData.h"
diff --git a/Engine/Source/Editor/UnrealEd/Private/EditorViewportCommands.cpp b/Engine/Source/Editor/UnrealEd/Private/EditorViewportCommands.cpp
--- a/Engine/Source/Editor/UnrealEd/Private/EditorViewportCommands.cpp
+++ b/Engine/Source/Editor/UnrealEd/Private/EditorViewportCommands.cpp
@@ -14,7 +14,6 @@
 #include "Materials/MaterialInterface.h"
 #include "MaterialShared.h"
 #include "Engine/Texture2D.h"
-#include "RayTracingDebugVisualizationMenuCommands.h"
 #include "GPUSkinCacheVisualizationMenuCommands.h"
 
 #define LOCTEXT_NAMESPACE "EditorViewportCommands"
@@ -71,7 +70,6 @@ void FEditorViewportCommands::RegisterCommands()
 
 	UI_COMMAND( PathTracingMode, "PathTracing View Mode", "Renders the scene using a path tracer", EUserInterfaceActionType::RadioButton, FInputChord());
 	UI_COMMAND( RayTracingDebugMode, "Ray tracing Debug View Mode", "Renders the scene in ray tracing debug mode", EUserInterfaceActionType::RadioButton, FInputChord());
-	FRayTracingDebugVisualizationMenuCommands::Register();
 
 	UI_COMMAND( DetailLightingMode, "Detail Lighting View Mode", "Renders the scene with detailed lighting only", EUserInterfaceActionType::RadioButton, FInputChord( EModifierKey::Alt, EKeys::Five ) );
 	UI_COMMAND( LightingOnlyMode, "Lighting Only View Mode", "Renders the scene with lights only, no textures", EUserInterfaceActionType::RadioButton, FInputChord( EModifierKey::Alt, EKeys::Six ) );
diff --git a/Engine/Source/Editor/UnrealEd/Private/RayTracingDebugVisualizationMenuCommands.cpp b/Engine/Source/Editor/UnrealEd/Private/RayTracingDebugVisualizationMenuCommands.cpp
deleted file mode 100644
diff --git a/Engine/Source/Editor/UnrealEd/Private/SEditorViewport.cpp b/Engine/Source/Editor/UnrealEd/Private/SEditorViewport.cpp
--- a/Engine/Source/Editor/UnrealEd/Private/SEditorViewport.cpp
+++ b/Engine/Source/Editor/UnrealEd/Private/SEditorViewport.cpp
@@ -22,7 +22,6 @@
 #include "Widgets/Input/SSpinBox.h"
 #include "Widgets/Input/SCheckBox.h"
 #include "MaterialShaderQualitySettings.h"
-#include "RayTracingDebugVisualizationMenuCommands.h"
 #include "GPUSkinCacheVisualizationMenuCommands.h"
 #include "GPUSkinCache.h"
 #include "Widgets/Colors/SComplexGradient.h"
@@ -658,9 +657,6 @@ void SEditorViewport::BindCommands()
 	{
 		MAP_VIEWMODE_ACTION(Commands.PathTracingMode, VMI_PathTracing);
 		MAP_VIEWMODE_ACTION(Commands.RayTracingDebugMode, VMI_RayTracingDebug);
-
-		const FRayTracingDebugVisualizationMenuCommands& RtDebugCommands = FRayTracingDebugVisualizationMenuCommands::Get();
-		RtDebugCommands.BindCommands(CommandListRef, Client);
 	}
 
 	MAP_VIEWMODE_ACTION( Commands.DetailLightingMode, VMI_Lit_DetailLighting );
diff --git a/Engine/Source/Editor/UnrealEd/Private/ViewportToolbar/UnrealEdViewportToolbar.cpp b/Engine/Source/Editor/UnrealEd/Private/ViewportToolbar/UnrealEdViewportToolbar.cpp
--- a/Engine/Source/Editor/UnrealEd/Private/ViewportToolbar/UnrealEdViewportToolbar.cpp
+++ b/Engine/Source/Editor/UnrealEd/Private/ViewportToolbar/UnrealEdViewportToolbar.cpp
@@ -19,7 +19,6 @@
 #include "ISettingsModule.h"
 #include "LevelEditor.h"
 #include "LevelEditorActions.h"
-#include "RayTracingDebugVisualizationMenuCommands.h"
 #include "SEditorViewport.h"
 #include "SScalabilitySettings.h"
 #include "ShowFlagMenuCommands.h"
@@ -2086,31 +2085,6 @@ void PopulateViewModesMenu(UToolMenu* InMenu)
 			}
 		}
 
-		if (IsRayTracingEnabled()
-			&& UE::UnrealEd::Private::ViewModesSubmenu::IsMenuSectionAvailable(
-				Context, EHidableViewModeMenuSections::RayTracingDebug
-			))
-		{
-			struct Local
-			{
-				static void BuildRayTracingDebugMenu(UToolMenu* InMenu)
-				{
-					const FRayTracingDebugVisualizationMenuCommands& RtDebugCommands =
-						FRayTracingDebugVisualizationMenuCommands::Get();
-					RtDebugCommands.BuildVisualisationSubMenu(InMenu);
-				}
-			};
-
-			Section.AddSubMenu(
-				"RayTracingDebugSubMenu",
-				LOCTEXT("RayTracingDebugSubMenu", "Ray Tracing Debug"),
-				LOCTEXT("RayTracing_ToolTip", "Select ray tracing buffer visualization view modes"),
-				FNewToolMenuDelegate::CreateStatic(&Local::BuildRayTracingDebugMenu),
-				false,
-				FSlateIcon(FAppStyle::GetAppStyleSetName(), "EditorViewport.RayTracingDebugMode")
-			);
-		}
-
 		{
 			struct Local
 			{
diff --git a/Engine/Source/Editor/UnrealEd/Public/RayTracingDebugVisualizationMenuCommands.h b/Engine/Source/Editor/UnrealEd/Public/RayTracingDebugVisualizationMenuCommands.h
deleted file mode 100644
diff --git a/Engine/Source/Editor/UnrealEd/Public/ViewportToolbar/UnrealEdViewportToolbar.h b/Engine/Source/Editor/UnrealEd/Public/ViewportToolbar/UnrealEdViewportToolbar.h
--- a/Engine/Source/Editor/UnrealEd/Public/ViewportToolbar/UnrealEdViewportToolbar.h
+++ b/Engine/Source/Editor/UnrealEd/Public/ViewportToolbar/UnrealEdViewportToolbar.h
@@ -33,8 +33,7 @@ namespace UE::UnrealEd
 enum EHidableViewModeMenuSections : uint8
 {
 	Exposure = 0,
-	GPUSkinCache = 1,
-	RayTracingDebug = 2
+	GPUSkinCache = 1
 };
 
 /** The value of this function is controlled by the CVAR "ToolMenusViewportToolbars". */
